Branching

//UE5/Private-Frosty/Collaboration/Frosty/Frosty/Plugins/FrostySensorGrid/...

to //UE5/Main/Engine/Plugins/Runtime/FrostySensorGrid/...

diff --git a/Engine/Plugins/Runtime/FrostySensorGrid/FrostySensorGrid.uplugin b/Engine/Plugins/Runtime/FrostySensorGrid/FrostySensorGrid.uplugin
new file mode 100644
--- /dev/null
+++ b/Engine/Plugins/Runtime/FrostySensorGrid/FrostySensorGrid.uplugin
@@ -0,0 +1,36 @@
+{
+	"FileVersion": 1,
+	"Version": 1,
+	"VersionName": "1.0",
+	"FriendlyName": "Frosty Sensor Grid",
+	"Description": "Niagara DataInterface for vehicle damage",
+	"Category": "VFX",
+	"CreatedBy" : "Epic Games Inc",
+	"CreatedByURL" : "http://epicgames.com",
+	"DocsURL": "",
+	"MarketplaceURL": "",
+	"SupportURL": "",
+	"EnabledByDefault" : false,
+	"CanContainContent": true,
+	"IsBetaVersion": false,
+	"IsExperimentalVersion": false,
+	"Installed": false,
+	"Modules": [
+		{
+			"Name": "FrostySensorGrid",
+			"Type": "Runtime",
+			"LoadingPhase": "Default"
+		},
+		{
+			"Name": "FrostySensorGridShaders",
+			"Type": "Runtime",
+			"LoadingPhase" : "PostConfigInit"
+		},
+	],
+	"Plugins": [
+		{
+			"Name": "Niagara",
+			"Enabled": true
+		}
+	]
+}
diff --git a/Engine/Plugins/Runtime/FrostySensorGrid/Shaders/FrostySensorGridAggregateBounds.usf b/Engine/Plugins/Runtime/FrostySensorGrid/Shaders/FrostySensorGridAggregateBounds.usf
new file mode 100644
--- /dev/null
+++ b/Engine/Plugins/Runtime/FrostySensorGrid/Shaders/FrostySensorGridAggregateBounds.usf
@@ -0,0 +1,58 @@
+// Copyright Epic Games, Inc. All Rights Reserved.
+
+#pragma once
+
+#include "/Engine/Public/Platform.ush"
+#include "FrostySensorGridCommon.ush"
+
+#define SUBGRID_SIZE 2
+
+RWByteAddressBuffer BoundsHierarchy;
+
+////////////////////////////////////////////////////////////////////////////////
+///// Shader to aggregate the bounds from one level of the hierarchy to the next
+////////////////////////////////////////////////////////////////////////////////
+int3 TargetStrideAndOffset;
+int3 SourceStrideAndOffset;
+
+groupshared float4 CachedBounds[SUBGRID_SIZE * SUBGRID_SIZE];
+
+[numthreads(SUBGRID_SIZE, SUBGRID_SIZE, 1)]
+void AggregateBounds(uint3 DispatchThreadId : SV_DispatchThreadID, uint GroupIndex : SV_GroupIndex, uint3 GroupID : SV_GroupID)
+{
+	CachedBounds[GroupIndex] = asfloat(BoundsHierarchy.Load4(SourceStrideAndOffset.z + 16 * TiledToLinear(DispatchThreadId, uint3(SourceStrideAndOffset.xy, 0))));
+
+	GroupMemoryBarrierWithGroupSync();
+
+	if (GroupIndex == 0)
+	{
+		float ValidCount = 0.0f;
+		float3 Centroid = float3(0.0f, 0.0f, 0.0f);
+
+		int i;
+
+		for (i = 0; i < SUBGRID_SIZE * SUBGRID_SIZE; ++i)
+		{
+			if (CachedBounds[i].w >= 0.0f)
+			{
+				Centroid += CachedBounds[i].xyz;
+				ValidCount += 1.0f;
+			}
+		}
+
+		Centroid *= (1.0f / ValidCount);
+
+		float MaxDistance = -1.0f;
+
+		for (i = 0; i < SUBGRID_SIZE * SUBGRID_SIZE; ++i)
+		{
+			if (CachedBounds[i].w >= 0.0f)
+			{
+				const float3 Delta = Centroid - CachedBounds[i].xyz;
+				MaxDistance = max(MaxDistance, length(Delta) + CachedBounds[i].w);
+			}
+		}
+
+		BoundsHierarchy.Store4(TargetStrideAndOffset.z + 16 * TiledToLinear(GroupID, uint3(TargetStrideAndOffset.xy, 0)), asuint(float4(Centroid, MaxDistance)));
+	}
+}
diff --git a/Engine/Plugins/Runtime/FrostySensorGrid/Shaders/FrostySensorGridCommon.ush b/Engine/Plugins/Runtime/FrostySensorGrid/Shaders/FrostySensorGridCommon.ush
new file mode 100644
--- /dev/null
+++ b/Engine/Plugins/Runtime/FrostySensorGrid/Shaders/FrostySensorGridCommon.ush
@@ -0,0 +1,15 @@
+// Copyright Epic Games, Inc. All Rights Reserved.
+
+#pragma once
+
+struct FSensorInfo
+{
+	float4 LocationAndDistance;
+	int3 HitIndex;
+	int _Pad0;
+};
+
+int TiledToLinear(int3 Index, int3 Dimensions)
+{
+	return Index.x + Dimensions.x * (Index.y + Dimensions.y * Index.z);
+}
\ No newline at end of file
diff --git a/Engine/Plugins/Runtime/FrostySensorGrid/Shaders/FrostySensorGridFindClosest.usf b/Engine/Plugins/Runtime/FrostySensorGrid/Shaders/FrostySensorGridFindClosest.usf
new file mode 100644
--- /dev/null
+++ b/Engine/Plugins/Runtime/FrostySensorGrid/Shaders/FrostySensorGridFindClosest.usf
@@ -0,0 +1,215 @@
+// Copyright Epic Games, Inc. All Rights Reserved.
+
+#pragma once
+
+#include "/Engine/Public/Platform.ush"
+#include "FrostySensorGridCommon.ush"
+
+#if !defined(MAX_HIERARCHY_LEVEL_COUNT)
+	#error MAX_HIERARCHY_LEVEL_COUNT is not defined
+#endif
+
+#if !defined(NETWORK_COUNT_PER_GROUP)
+	#error NETWORK_COUNT_PER_GROUP is not defined
+#endif
+
+#define SUBGRID_SIZE 2
+#define STACK_STRIDE (((MAX_HIERARCHY_LEVEL_COUNT - 1) * 3) + 1)
+
+////////////////////////////////////////////////////////////////////////////////
+///// Searches through the bounds hierarchy to find the closest valid sensor
+////////////////////////////////////////////////////////////////////////////////
+
+ByteAddressBuffer SensorHierarchy;
+RWStructuredBuffer<FSensorInfo> NearestSensors;
+int3 SensorGridDimensions;
+int HierarchyLevelOffsets[MAX_HIERARCHY_LEVEL_COUNT];
+float2 DistanceBounds;
+int HierarchyLevelCount;
+
+// CellIndex -> GridX, GridY, GridZ
+// Level (0 is largest 'mip', 1 is half the size, ...)
+float4 RetrieveBounds(uint3 CellIndex, uint HierarchyLevel)
+{
+	const uint CellCount = 1 << (HierarchyLevelCount - HierarchyLevel - 1);
+	const uint BoundsAddress = HierarchyLevelOffsets[HierarchyLevel] + 16 * TiledToLinear(CellIndex, int3(CellCount, CellCount, 0));
+
+	return asfloat(SensorHierarchy.Load4(BoundsAddress));
+}
+
+float DistanceBetweenBounds(float4 A, float4 B)
+{
+	return length(B.xyz - A.xyz) - A.w - B.w;
+}
+
+groupshared int3 ThreadBestSensorIndex[SUBGRID_SIZE * SUBGRID_SIZE * NETWORK_COUNT_PER_GROUP];
+groupshared float ThreadBestSensorDistance[SUBGRID_SIZE * SUBGRID_SIZE * NETWORK_COUNT_PER_GROUP];
+
+groupshared float4 ThreadLeaves[SUBGRID_SIZE][SUBGRID_SIZE];
+groupshared float4 ThreadBounds;
+
+[numthreads(1, 1, NETWORK_COUNT_PER_GROUP)]
+void FindClosestSensors(uint3 GroupThreadId : SV_GroupThreadID, uint GroupIndex : SV_GroupIndex, uint3 GroupId : SV_GroupID)
+{
+	if (GroupId.z >= uint(SensorGridDimensions.z))
+	{
+		return;
+	}
+
+	// collect the bounds of the leaves and their immmediate parent
+	if (GroupIndex == 0)
+	{
+		ThreadBounds = RetrieveBounds(GroupId, log2(SUBGRID_SIZE));
+		for (int i = 0; i < SUBGRID_SIZE; ++i)
+		{
+			for (int j = 0; j < SUBGRID_SIZE; ++j)
+			{
+				ThreadLeaves[i][j] = RetrieveBounds(GroupId * uint3(2, 2, 1) + uint3(i, j, 0), 0);
+
+				for (int o = 0; o < NETWORK_COUNT_PER_GROUP; ++o)
+				{
+					ThreadBestSensorDistance[o + NETWORK_COUNT_PER_GROUP * (j + SUBGRID_SIZE * i)] = DistanceBounds.y;
+				}
+			}
+		}
+	}
+
+	GroupMemoryBarrierWithGroupSync();
+
+	const float MinDistance = DistanceBounds.x;
+	float WorstLocalDistance = DistanceBounds.y;
+
+	for (int OwnerOffset = 1; OwnerOffset <= SensorGridDimensions.z; OwnerOffset += NETWORK_COUNT_PER_GROUP)
+	{
+		const uint OwnerIndex = (GroupId.z + GroupIndex + OwnerOffset) % SensorGridDimensions.z;
+
+		if (OwnerIndex == GroupId.z || WorstLocalDistance < MinDistance)
+		{
+			continue;
+		}
+
+		int4 DepthFirstStack[STACK_STRIDE];
+		int StackDepth = 0;
+
+		DepthFirstStack[StackDepth++] = int4(0, 0, OwnerIndex, HierarchyLevelCount - 1);
+
+		while (StackDepth > 0 && WorstLocalDistance > MinDistance)
+		{
+			const int4 ParentIndex = DepthFirstStack[--StackDepth];
+			const float4 ParentBounds = RetrieveBounds(ParentIndex.xyz, ParentIndex.w);
+
+			float Delta = DistanceBetweenBounds(ParentBounds, ThreadBounds);
+
+			// if no part of this cell is closer than our current best case sensor, then we can skip it
+			if (Delta < WorstLocalDistance)
+			{
+				// if the bounds of the parent are smaller than our min distance we'll just use the first
+				// leaf node underneath us (all others are within our specified error tolerance)
+				if (ParentBounds.w <= MinDistance)
+				{
+					int3 LeafIndex = int3(ParentIndex.xy * (1 << ParentIndex.w), ParentIndex.z);
+					float MaxLocalDistances = 0.0f;
+
+					for (int ThreadI = 0; ThreadI < SUBGRID_SIZE; ++ThreadI)
+					{
+						for (int ThreadJ = 0; ThreadJ < SUBGRID_SIZE; ++ThreadJ)
+						{
+							const int ThreadRecordIndex = GroupThreadId.z + NETWORK_COUNT_PER_GROUP * (ThreadJ + SUBGRID_SIZE * ThreadI);
+							const float4 ThreadLeafBounds = ThreadLeaves[ThreadI][ThreadJ];
+
+							const float LeafDistance = DistanceBetweenBounds(ParentBounds, ThreadLeafBounds);
+
+							if (LeafDistance < ThreadBestSensorDistance[ThreadRecordIndex])
+							{
+								ThreadBestSensorDistance[ThreadRecordIndex] = LeafDistance;
+								ThreadBestSensorIndex[ThreadRecordIndex] = LeafIndex;
+							}
+
+							MaxLocalDistances = max(MaxLocalDistances, ThreadBestSensorDistance[ThreadRecordIndex]);
+						}
+					}
+
+					WorstLocalDistance = min(MaxLocalDistances, DistanceBounds.y);
+				}
+				// see if we have more intermediary nodes to get through
+				else if (ParentIndex.w > 1)
+				{
+					// queue up potential candidates
+					DepthFirstStack[StackDepth++] = ParentIndex * int4(2, 2, 1, 1) + int4(1, 1, 0, -1); // BR
+					DepthFirstStack[StackDepth++] = ParentIndex * int4(2, 2, 1, 1) + int4(0, 1, 0, -1); // BL
+					DepthFirstStack[StackDepth++] = ParentIndex * int4(2, 2, 1, 1) + int4(1, 0, 0, -1); // TR
+					DepthFirstStack[StackDepth++] = ParentIndex * int4(2, 2, 1, 1) + int4(0, 0, 0, -1); // TL
+				}
+				else
+				{
+					float MaxLocalDistances = 0.0f;
+
+					// perform the actual comparison between leaves
+					for (int ThreadI = 0; ThreadI < SUBGRID_SIZE; ++ThreadI)
+					{
+						for (int ThreadJ = 0; ThreadJ < SUBGRID_SIZE; ++ThreadJ)
+						{
+							const int ThreadRecordIndex = GroupThreadId.z + NETWORK_COUNT_PER_GROUP * (ThreadJ + SUBGRID_SIZE * ThreadI);
+							const float4 ThreadLeafBounds = ThreadLeaves[ThreadI][ThreadJ];
+
+							if (DistanceBetweenBounds(ParentBounds, ThreadLeafBounds) < WorstLocalDistance)
+							{
+								for (int LeafI = 0; LeafI < SUBGRID_SIZE; ++LeafI)
+								{
+									for (int LeafJ = 0; LeafJ < SUBGRID_SIZE; ++LeafJ)
+									{
+										const int3 LeafIndex = ParentIndex.xyz * int3(2, 2, 1) + int3(LeafI, LeafJ, 0);
+										const float4 LeafBounds = RetrieveBounds(LeafIndex, 0);
+										const float LeafDistance = DistanceBetweenBounds(ThreadLeafBounds, LeafBounds);
+
+										if (LeafDistance < ThreadBestSensorDistance[ThreadRecordIndex])
+										{
+											ThreadBestSensorDistance[ThreadRecordIndex] = LeafDistance;
+											ThreadBestSensorIndex[ThreadRecordIndex] = LeafIndex;
+										}
+									}
+								}
+							}
+
+							MaxLocalDistances = max(MaxLocalDistances, ThreadBestSensorDistance[ThreadRecordIndex]);
+						}
+					}
+
+					WorstLocalDistance = min(MaxLocalDistances, DistanceBounds.y);
+				}
+			}
+		}
+	}
+
+	GroupMemoryBarrierWithGroupSync();
+
+	if (GroupIndex == 0)
+	{
+		for (int i = 0; i < SUBGRID_SIZE; ++i)
+		{
+			for (int j = 0; j < SUBGRID_SIZE; ++j)
+			{
+				FSensorInfo BestSensor;
+				BestSensor.HitIndex = int3(-1, -1, -1);
+				BestSensor._Pad0 = 0;
+				BestSensor.LocationAndDistance = float4(0.0f, 0.0f, 0.0f, DistanceBounds.y);
+
+				const int StartThreadIndex = NETWORK_COUNT_PER_GROUP * (j + SUBGRID_SIZE * i);
+				for (int OwnerIndex = 0; OwnerIndex < NETWORK_COUNT_PER_GROUP; ++OwnerIndex)
+				{
+					if (ThreadBestSensorDistance[StartThreadIndex + OwnerIndex] < BestSensor.LocationAndDistance.w)
+					{
+						BestSensor.HitIndex = ThreadBestSensorIndex[StartThreadIndex + OwnerIndex];
+						BestSensor.LocationAndDistance = float4(RetrieveBounds(BestSensor.HitIndex, 0).xyz, ThreadBestSensorDistance[StartThreadIndex + OwnerIndex]);
+					}
+				}
+
+				// write the results
+				const uint ResultIndex = TiledToLinear(GroupId * uint3(2, 2, 1) + uint3(i, j, 0), SensorGridDimensions);
+				NearestSensors[ResultIndex] = BestSensor;
+			}
+		}
+	}
+}
+
+
diff --git a/Engine/Plugins/Runtime/FrostySensorGrid/Shaders/NiagaraDataInterfaceFrostySensorGrid.ush b/Engine/Plugins/Runtime/FrostySensorGrid/Shaders/NiagaraDataInterfaceFrostySensorGrid.ush
new file mode 100644
--- /dev/null
+++ b/Engine/Plugins/Runtime/FrostySensorGrid/Shaders/NiagaraDataInterfaceFrostySensorGrid.ush
@@ -0,0 +1,39 @@
+// Copyright Epic Games, Inc. All Rights Reserved.
+
+#pragma once
+
+#include "FrostySensorGridCommon.ush"
+
+////////////////////////////////////////////////////////////////////////////////
+///// Writes a particles location to the BoundsHierarchy
+////////////////////////////////////////////////////////////////////////////////
+void NDIFrostySensorGrid_UpdateSensorGpu(int3 SensorGridDimensions, int3 SensorIndex, float3 SensorLocation, float SensorRange, bool IsValid, RWByteAddressBuffer BoundsHierarchy)
+{
+	if (!IsValid)
+	{
+		SensorLocation = float3(0.0f, 0.0f, 0.0f);
+		SensorRange = -1.0f;
+	}
+
+	BoundsHierarchy.Store4(16 * TiledToLinear(SensorIndex, SensorGridDimensions), asuint(float4(SensorLocation, SensorRange)));
+}
+
+////////////////////////////////////////////////////////////////////////////////
+///// Searches the SensorInfo for previously calculated closest sensor
+////////////////////////////////////////////////////////////////////////////////
+void NDIFrostySensorGrid_FindNearestGpu(StructuredBuffer<FSensorInfo> SensorInfo, int3 SensorGridDimensions, int3 SourceSensorIndex, out float3 Location, out float Distance, out bool IsValid)
+{
+	if (all(SourceSensorIndex < SensorGridDimensions))
+	{
+		const int LinearSensorIndex = TiledToLinear(SourceSensorIndex, SensorGridDimensions);
+
+		Location = SensorInfo[LinearSensorIndex].LocationAndDistance.xyz;
+		Distance = SensorInfo[LinearSensorIndex].LocationAndDistance.w;
+
+		IsValid = all(SensorInfo[LinearSensorIndex].HitIndex >= 0);
+	}
+	else
+	{
+		IsValid = false;
+	}
+}
\ No newline at end of file
diff --git a/Engine/Plugins/Runtime/FrostySensorGrid/Shaders/NiagaraDataInterfaceFrostySensorGridTemplate.ush b/Engine/Plugins/Runtime/FrostySensorGrid/Shaders/NiagaraDataInterfaceFrostySensorGridTemplate.ush
new file mode 100644
--- /dev/null
+++ b/Engine/Plugins/Runtime/FrostySensorGrid/Shaders/NiagaraDataInterfaceFrostySensorGridTemplate.ush
@@ -0,0 +1,25 @@
+// Copyright Epic Games, Inc. All Rights Reserved.
+
+RWByteAddressBuffer RWBoundsHierarchy_{ParameterName};
+StructuredBuffer<FSensorInfo> SensorInfo_{ParameterName};
+int3 SensorGridDimensions_{ParameterName};
+int SensorGridWriteIndex_{ParameterName};
+int SensorGridReadIndex_{ParameterName};
+
+void UpdateSensorGpu_{ParameterName}_UEImpureCall(int SensorX, int SensorY, float3 SensorLocation, float SensorRange, bool IsValid)
+{
+	NDIFrostySensorGrid_UpdateSensorGpu(SensorGridDimensions_{ParameterName}, int3(SensorX, SensorY, SensorGridWriteIndex_{ParameterName}), SensorLocation, SensorRange, IsValid, RWBoundsHierarchy_{ParameterName});
+}
+
+void FindNearestGpu_{ParameterName}(int SensorX, int SensorY, out float3 Location, out float Distance, out bool IsValid)
+{
+	Location = float3(0.0f, 0.0f, 0.0f);
+	Distance = 0.0f;
+	IsValid = false;
+
+	if (SensorGridReadIndex_{ParameterName} >= 0)
+	{
+		int3 TargetSensorIndex;
+		NDIFrostySensorGrid_FindNearestGpu(SensorInfo_{ParameterName}, SensorGridDimensions_{ParameterName}, int3(SensorX, SensorY, SensorGridReadIndex_{ParameterName}), Location, Distance, IsValid);
+	}
+}
\ No newline at end of file
diff --git a/Engine/Plugins/Runtime/FrostySensorGrid/Source/FrostySensorGrid/FrostySensorGrid.Build.cs b/Engine/Plugins/Runtime/FrostySensorGrid/Source/FrostySensorGrid/FrostySensorGrid.Build.cs
new file mode 100644
--- /dev/null
+++ b/Engine/Plugins/Runtime/FrostySensorGrid/Source/FrostySensorGrid/FrostySensorGrid.Build.cs
@@ -0,0 +1,37 @@
+// Copyright Epic Games, Inc. All Rights Reserved.
+
+namespace UnrealBuildTool.Rules
+{
+	public class FrostySensorGrid : ModuleRules
+	{
+		public FrostySensorGrid(ReadOnlyTargetRules Target) : base(Target)
+		{
+			// Include Renderer/Private to have access to default resources
+			PrivateIncludePaths.AddRange(
+				new string[] {
+					ModuleDirectory + "/Private",
+					EngineDirectory + "/Source/Runtime/Renderer/Private",
+				});
+			PublicIncludePaths.Add(ModuleDirectory + "/Public");
+
+			PrivateDependencyModuleNames.AddRange(
+				new string[]
+				{
+					"Core",
+					"CoreUObject",
+					"Engine",
+					"FrostySensorGridShaders",
+					"Projects",
+					"MeshDescription",
+					"MovieScene",
+					"NiagaraCore",
+					"Niagara",
+					"NiagaraShader",
+					"RenderCore",
+					"Renderer",
+					"VectorVM",
+					"RHI",
+				});
+		}
+	}
+}
diff --git a/Engine/Plugins/Runtime/FrostySensorGrid/Source/FrostySensorGrid/Private/FrostySensorGrid.cpp b/Engine/Plugins/Runtime/FrostySensorGrid/Source/FrostySensorGrid/Private/FrostySensorGrid.cpp
new file mode 100644
--- /dev/null
+++ b/Engine/Plugins/Runtime/FrostySensorGrid/Source/FrostySensorGrid/Private/FrostySensorGrid.cpp
@@ -0,0 +1,14 @@
+// Copyright Epic Games, Inc. All Rights Reserved.
+
+#include "FrostySensorGrid.h"
+#include "Interfaces/IPluginManager.h"
+
+IMPLEMENT_MODULE(FFrostySensorGrid, FrostySensorGrid);
+
+void FFrostySensorGrid::StartupModule()
+{
+}
+
+void FFrostySensorGrid::ShutdownModule()
+{
+}
\ No newline at end of file
diff --git a/Engine/Plugins/Runtime/FrostySensorGrid/Source/FrostySensorGrid/Private/NiagaraDataInterfaceFrostySensorGrid.cpp b/Engine/Plugins/Runtime/FrostySensorGrid/Source/FrostySensorGrid/Private/NiagaraDataInterfaceFrostySensorGrid.cpp
new file mode 100644
--- /dev/null
+++ b/Engine/Plugins/Runtime/FrostySensorGrid/Source/FrostySensorGrid/Private/NiagaraDataInterfaceFrostySensorGrid.cpp
@@ -0,0 +1,632 @@
+// Copyright Epic Games, Inc. All Rights Reserved.
+
+#include "NiagaraDataInterfaceFrostySensorGrid.h"
+
+#include "FrostySensorGridShaders.h"
+#include "GlobalDistanceFieldParameters.h"
+#include "NiagaraComponent.h"
+#include "NiagaraEmitterInstanceBatcher.h"
+#include "NiagaraRayTracingHelper.h"
+#include "NiagaraTypes.h"
+#include "NiagaraWorldManager.h"
+#include "RayTracingInstanceUtils.h"
+#include "RenderResource.h"
+#include "Shader.h"
+#include "ShaderCore.h"
+#include "ShaderCompilerCore.h"
+#include "ShaderParameterUtils.h"
+
+// c++ mirror of the struct defined in NiagaraFrostySensorGridCommon.ush
+struct alignas(16) FSensorInfo
+{
+	FVector4 LocationAndDistance;
+	FIntVector HitIndex;
+};
+
+#define LOCTEXT_NAMESPACE "NiagaraDataInterfaceFrostySensorGrid"
+
+namespace NDIFrostySensorGridLocal
+{
+	static const TCHAR* CommonShaderFile = TEXT("/FrostySensorGrid/NiagaraDataInterfaceFrostySensorGrid.ush");
+	static const TCHAR* TemplateShaderFile = TEXT("/FrostySensorGrid/NiagaraDataInterfaceFrostySensorGridTemplate.ush");
+
+	static const FName UpdateSensorName(TEXT("UpdateSensorGpu"));
+	static const FName FindNearestName(TEXT("FindNearestGpu"));
+
+	static const FString BoundsHierarchyParamName(TEXT("BoundsHierarchy_"));
+	static const FString SensorInfoParamName(TEXT("SensorInfo_"));
+	static const FString SensorGridDimensionsParamName(TEXT("SensorGridDimensions_"));
+	static const FString SensorGridWriteIndexParamName(TEXT("SensorGridWriteIndex_"));
+	static const FString SensorGridReadIndexParamName(TEXT("SensorGridReadIndex_"));
+
+	enum NodeVersion : int32
+	{
+		NodeVersion_Initial = 0,
+		Nodeversion_Reworked_Output,
+
+		VersionPlusOne,
+		NodeVersion_Latest = VersionPlusOne - 1,
+	};
+
+	const FText SensorXDescription = IF_WITH_EDITORONLY_DATA(LOCTEXT("SensorXDescription", "?."), FText());
+	const FText SensorYDescription = IF_WITH_EDITORONLY_DATA(LOCTEXT("SensorYDescription", "??."), FText());
+	const FText LocationDescription = IF_WITH_EDITORONLY_DATA(LOCTEXT("LocationDescription", "????."), FText());
+	const FText SensorRangeDescription = IF_WITH_EDITORONLY_DATA(LOCTEXT("SensorRangeDescription", "????."), FText());
+	const FText SensorValidDescription = IF_WITH_EDITORONLY_DATA(LOCTEXT("SensorValidDescription", "????."), FText());
+
+	const FText OutputLocationDescription = IF_WITH_EDITORONLY_DATA(LOCTEXT("OutputLocationDescription", "?."), FText());
+	const FText OutputDistanceDescription = IF_WITH_EDITORONLY_DATA(LOCTEXT("OutputDistanceDescription", "?."), FText());
+	const FText OutputIsValidDescription = IF_WITH_EDITORONLY_DATA(LOCTEXT("OutputIsValidDescription", "?."), FText());
+
+}
+
+struct FSensorGridNetworkProxy
+{
+	// ByteAddress buffer used to store the bounds hierarchy for all sensors in the network
+	//	-the leaves of the trees are populated by the simulation pass of the system instances
+	//	-global shader used to construct the higher level tiers of the tree
+	TUniquePtr<FRWByteAddressBuffer> BoundsHierarchy;
+
+	// StructuredBuffer storing the results of a global query for finding the closest sensor
+	//	-populated by global shader
+	//	-read by simulation pass of the system instances
+	TUniquePtr<FRWBufferStructured> SensorInfo;
+
+	TUniquePtr<FRWBufferStructured> EmptyResults;
+
+	// mapping between the system instance ID and the subnetwork within the buffers
+	TMap<FNiagaraSystemInstanceID, int32> InstanceOwnerReadIndexMap;
+	TMap<FNiagaraSystemInstanceID, int32> InstanceOwnerWriteIndexMap;
+
+	TSet<FNiagaraSystemInstanceID> RegisteredInstances;
+
+	int32 QueuedOwnerCount = 0;
+	int32 AllocatedOwnerCount = 0;
+	int32 ResultsOwnerCount = 0;
+	uint32 ProcessedFrameNumber = INDEX_NONE;
+
+	bool CurrentlySimulating = false;
+
+	// XY is the grid of sensors
+	const FIntPoint SensorGridDimensions;
+
+	FSensorGridNetworkProxy() = delete;
+	FSensorGridNetworkProxy(int32 SensorGridWidth)
+		: SensorGridDimensions(SensorGridWidth, SensorGridWidth)
+	{
+		// initialize with a dummy SensorInfo structured buffer which will be used for the frames where we don't have any results
+		EmptyResults = MakeUnique<FRWBufferStructured>();
+		EmptyResults->Initialize(
+			TEXT("FrostySensorGridEmptyResults"),
+			sizeof(FSensorInfo),
+			1,
+			BUF_Static,
+			false /*bUseUavCounter*/,
+			false /*bAppendBuffer*/,
+			ERHIAccess::SRVCompute);
+	}
+
+	~FSensorGridNetworkProxy()
+	{
+		if (auto* BoundsHierarchyRef = BoundsHierarchy.Release())
+		{
+			BoundsHierarchyRef->Release();
+			delete BoundsHierarchyRef;
+		}
+
+		if (auto* SensorInfoRef = SensorInfo.Release())
+		{
+			SensorInfoRef->Release();
+			delete SensorInfoRef;
+		}
+
+		if (auto* EmptyResultsRef = EmptyResults.Release())
+		{
+			EmptyResultsRef->Release();
+			delete EmptyResultsRef;
+		}
+	}
+
+	void BeginSimulation(FRHICommandList& RHICmdList)
+	{
+		if (QueuedOwnerCount)
+		{
+			if (QueuedOwnerCount != AllocatedOwnerCount)
+			{
+				if (FRWByteAddressBuffer* ReleasedBuffer = BoundsHierarchy.Release())
+				{
+					ReleasedBuffer->Release();
+					delete ReleasedBuffer;
+				}
+
+				int32 HierarchySensorCount = 0;
+				for (int32 SensorIt = 1; SensorIt <= SensorGridDimensions.X; SensorIt *= 2)
+				{
+					HierarchySensorCount += SensorIt * SensorIt * QueuedOwnerCount;
+				}
+
+				BoundsHierarchy = MakeUnique<FRWByteAddressBuffer>();
+				BoundsHierarchy->Initialize(
+					TEXT("FrostySensorGridBoundsHierarchy"),
+					HierarchySensorCount * sizeof(FVector4),
+					BUF_Static);
+
+				RHICmdList.Transition(FRHITransitionInfo(BoundsHierarchy->UAV, ERHIAccess::SRVMask, ERHIAccess::UAVCompute));
+			}
+
+			AllocatedOwnerCount = QueuedOwnerCount;
+			QueuedOwnerCount = 0;
+
+			RHICmdList.BeginUAVOverlap(BoundsHierarchy->UAV);
+			CurrentlySimulating = true;
+		}
+	}
+
+	void EndSimulation(FRHICommandList& RHICmdList)
+	{
+		if (AllocatedOwnerCount && CurrentlySimulating)
+		{
+			RHICmdList.EndUAVOverlap(BoundsHierarchy->UAV);
+			CurrentlySimulating = false;
+		}
+	}
+
+	void PrepareResultsBuffer()
+	{
+		if (AllocatedOwnerCount != ResultsOwnerCount)
+		{
+			if (FRWBufferStructured* ReleasedBuffer = SensorInfo.Release())
+			{
+				ReleasedBuffer->Release();
+				delete ReleasedBuffer;
+			}
+
+			ResultsOwnerCount = AllocatedOwnerCount;
+
+			SensorInfo = MakeUnique<FRWBufferStructured>();
+			SensorInfo->Initialize(
+				TEXT("FrostySensorGridSensorInfo"),
+				sizeof(FSensorInfo),
+				SensorGridDimensions.X * SensorGridDimensions.Y * ResultsOwnerCount,
+				BUF_Static,
+				false /*bUseUavCounter*/,
+				false /*bAppendBuffer*/,
+				ERHIAccess::SRVCompute);
+		}
+	}
+
+	void Dispatch(FRHICommandList& RHICmdList, ERHIFeatureLevel::Type InFeatureLevel, const FIntPoint& InSensorGridDimensions, const FVector2D& GlobalSensorRange)
+	{
+		if (ProcessedFrameNumber == GFrameNumberRenderThread)
+		{
+			return;
+		}
+
+		ProcessedFrameNumber = GFrameNumberRenderThread;
+
+		if (AllocatedOwnerCount)
+		{
+			PrepareResultsBuffer();
+
+			FFrostySensorGridHelper Helper(InFeatureLevel, FIntVector(SensorGridDimensions.X, SensorGridDimensions.Y, AllocatedOwnerCount));
+
+			RHICmdList.Transition(FRHITransitionInfo(BoundsHierarchy->UAV, ERHIAccess::Unknown, ERHIAccess::UAVCompute));
+			Helper.BuildBounds(RHICmdList, BoundsHierarchy->UAV);
+
+			RHICmdList.Transition(FRHITransitionInfo(BoundsHierarchy->UAV, ERHIAccess::Unknown, ERHIAccess::SRVCompute));
+			RHICmdList.Transition(FRHITransitionInfo(SensorInfo->UAV, ERHIAccess::SRVCompute, ERHIAccess::UAVCompute));
+			Helper.FindNearestSensors(RHICmdList, BoundsHierarchy->SRV, GlobalSensorRange, SensorInfo->UAV);
+			RHICmdList.Transition(FRHITransitionInfo(SensorInfo->UAV, ERHIAccess::Unknown, ERHIAccess::SRVCompute));
+
+			InstanceOwnerReadIndexMap = InstanceOwnerWriteIndexMap;
+			InstanceOwnerWriteIndexMap.Reset();
+		}
+	}
+};
+
+struct FNiagaraDataIntefaceProxyFrostySensorGrid : public FNiagaraDataInterfaceProxy
+{
+	FRWLock NetworkLock;
+	TMap<const NiagaraEmitterInstanceBatcher*, TUniquePtr<FSensorGridNetworkProxy>> NetworkProxies;
+
+	int32 SensorGridSize = 0;
+	FVector2D GlobalSensorRange = FVector2D(EForceInit::ForceInitToZero);
+
+	virtual int32 PerInstanceDataPassedToRenderThreadSize() const override
+	{
+		return 0;
+	}
+
+	FSensorGridNetworkProxy* GetNetwork(const NiagaraEmitterInstanceBatcher* Batcher, bool CreateIfMissing = false)
+	{
+		FReadScopeLock _Scope(NetworkLock);
+		return NetworkProxies.FindChecked(Batcher).Get();
+	}
+
+	void RegisterNetworkInstance(const NiagaraEmitterInstanceBatcher* Batcher, FNiagaraSystemInstanceID SystemInstanceID)
+	{
+		FWriteScopeLock _Scope(NetworkLock);
+		TUniquePtr<FSensorGridNetworkProxy>& Network = NetworkProxies.FindOrAdd(Batcher, nullptr);
+		if (!Network)
+		{
+			Network = MakeUnique<FSensorGridNetworkProxy>(SensorGridSize);
+		}
+
+		Network->RegisteredInstances.Add(SystemInstanceID);
+	}
+
+	void UnregisterNetworkInstance(const NiagaraEmitterInstanceBatcher* Batcher, FNiagaraSystemInstanceID SystemInstanceID)
+	{
+		FWriteScopeLock _Scope(NetworkLock);
+		if (TUniquePtr<FSensorGridNetworkProxy>* NetworkPtr = NetworkProxies.Find(Batcher))
+		{
+			if (FSensorGridNetworkProxy* Network = NetworkPtr->Get())
+			{
+				Network->RegisteredInstances.Remove(SystemInstanceID);
+				if (!Network->RegisteredInstances.Num())
+				{
+					NetworkProxies.Remove(Batcher);
+				}
+			}
+
+		}
+	}
+
+	virtual void PreStage(FRHICommandList& RHICmdList, const FNiagaraDataInterfaceStageArgs& Context) override
+	{
+		FNiagaraDataInterfaceProxy::PreStage(RHICmdList, Context);
+
+		if (FSensorGridNetworkProxy* Network = GetNetwork(Context.Batcher, true))
+		{
+			int32& OwnerIndex = Network->InstanceOwnerWriteIndexMap.FindOrAdd(Context.SystemInstanceID, INDEX_NONE);
+			if (OwnerIndex == INDEX_NONE)
+			{
+				OwnerIndex = Network->QueuedOwnerCount++;
+			}
+		}
+	}
+
+	virtual void PostSimulate(FRHICommandList& RHICmdList, const FNiagaraDataInterfaceArgs& Context) override
+	{
+		FNiagaraDataInterfaceProxy::PostSimulate(RHICmdList, Context);
+
+		if (FSensorGridNetworkProxy* Network = GetNetwork(Context.Batcher))
+		{
+			Network->Dispatch(RHICmdList, Context.Batcher->GetFeatureLevel(), Network->SensorGridDimensions, GlobalSensorRange);
+		}
+	}
+
+	void RenderThreadInitialize(int32 InSensorCount, const FVector2D& InGlobalSensorRange)
+	{
+		// SensorCount needs to be a power of two
+		SensorGridSize = FMath::RoundUpToPowerOfTwo(InSensorCount);
+		GlobalSensorRange.X = FMath::Max(0.0f, InGlobalSensorRange.X);
+		GlobalSensorRange.Y = FMath::Max(GlobalSensorRange.X, InGlobalSensorRange.Y);
+
+		NetworkProxies.Reset();
+	}
+};
+
+
+UNiagaraDataInterfaceFrostySensorGrid::UNiagaraDataInterfaceFrostySensorGrid(FObjectInitializer const& ObjectInitializer)
+	: Super(ObjectInitializer)
+{
+    Proxy.Reset(new FNiagaraDataIntefaceProxyFrostySensorGrid());
+}
+
+void UNiagaraDataInterfaceFrostySensorGrid::PostInitProperties()
+{
+	Super::PostInitProperties();
+
+	//Can we register data interfaces as regular types and fold them into the FNiagaraVariable framework for UI and function calls etc?
+	if (HasAnyFlags(RF_ClassDefaultObject))
+	{
+		ENiagaraTypeRegistryFlags Flags = ENiagaraTypeRegistryFlags::AllowAnyVariable | ENiagaraTypeRegistryFlags::AllowParameter;
+		FNiagaraTypeRegistry::Register(FNiagaraTypeDefinition(GetClass()), Flags);
+	}
+}
+
+void UNiagaraDataInterfaceFrostySensorGrid::PostLoad()
+{
+	Super::PostLoad();
+
+	if (SensorCountPerSide)
+	{
+		MarkRenderDataDirty();
+	}
+}
+
+void UNiagaraDataInterfaceFrostySensorGrid::GetFunctions(TArray<FNiagaraFunctionSignature>& OutFunctions)
+{
+	FNiagaraFunctionSignature SigUpdateSensor;
+	SigUpdateSensor.Name = NDIFrostySensorGridLocal::UpdateSensorName;
+	SigUpdateSensor.bRequiresExecPin = true;
+	SigUpdateSensor.bMemberFunction = true;
+	SigUpdateSensor.bRequiresContext = false;
+	SigUpdateSensor.bSupportsCPU = false;
+#if WITH_EDITORONLY_DATA
+	SigUpdateSensor.FunctionVersion = NDIFrostySensorGridLocal::NodeVersion_Latest;
+	SigUpdateSensor.Description = LOCTEXT("SceneDepthSignatureDescription", "Projects a given world position to view space and then queries the depth buffer with that position.");
+#endif
+
+	SigUpdateSensor.Inputs.Add(FNiagaraVariable(FNiagaraTypeDefinition(GetClass()), TEXT("FrostySensorGrid")));
+	SigUpdateSensor.AddInput(FNiagaraVariable(FNiagaraTypeDefinition::GetIntDef(), TEXT("Sensor_X")), NDIFrostySensorGridLocal::SensorXDescription);
+	SigUpdateSensor.AddInput(FNiagaraVariable(FNiagaraTypeDefinition::GetIntDef(), TEXT("Sensor_Y")), NDIFrostySensorGridLocal::SensorYDescription);
+	SigUpdateSensor.AddInput(FNiagaraVariable(FNiagaraTypeDefinition::GetVec3Def(), TEXT("Location")), NDIFrostySensorGridLocal::LocationDescription);
+	SigUpdateSensor.AddInput(FNiagaraVariable(FNiagaraTypeDefinition::GetFloatDef(), TEXT("SensorRange")), NDIFrostySensorGridLocal::SensorRangeDescription);
+	SigUpdateSensor.AddInput(FNiagaraVariable(FNiagaraTypeDefinition::GetBoolDef(), TEXT("IsValid")), NDIFrostySensorGridLocal::SensorValidDescription);
+	OutFunctions.Add(SigUpdateSensor);
+
+	FNiagaraFunctionSignature SigFindNearest;
+	SigFindNearest.Name = NDIFrostySensorGridLocal::FindNearestName;
+	SigFindNearest.bMemberFunction = true;
+	SigFindNearest.bRequiresContext = false;
+	SigFindNearest.bSupportsCPU = false;
+#if WITH_EDITORONLY_DATA
+	SigFindNearest.FunctionVersion = NDIFrostySensorGridLocal::NodeVersion_Latest;
+	SigFindNearest.Description = LOCTEXT("CustomDepthDescription", "Projects a given world position to view space and then queries the custom depth buffer with that position.");
+#endif
+
+	SigFindNearest.Inputs.Add(FNiagaraVariable(FNiagaraTypeDefinition(GetClass()), TEXT("FrostySensorGrid")));
+	SigFindNearest.AddInput(FNiagaraVariable(FNiagaraTypeDefinition::GetIntDef(), TEXT("Sensor_X")), NDIFrostySensorGridLocal::SensorXDescription);
+	SigFindNearest.AddInput(FNiagaraVariable(FNiagaraTypeDefinition::GetIntDef(), TEXT("Sensor_Y")), NDIFrostySensorGridLocal::SensorYDescription);
+	SigFindNearest.AddOutput(FNiagaraVariable(FNiagaraTypeDefinition::GetVec3Def(), TEXT("Out_Location")), NDIFrostySensorGridLocal::OutputLocationDescription);
+	SigFindNearest.AddOutput(FNiagaraVariable(FNiagaraTypeDefinition::GetFloatDef(), TEXT("Out_Distance")), NDIFrostySensorGridLocal::OutputDistanceDescription);
+	SigFindNearest.AddOutput(FNiagaraVariable(FNiagaraTypeDefinition::GetBoolDef(), TEXT("Out_IsValid")), NDIFrostySensorGridLocal::OutputIsValidDescription);
+	OutFunctions.Add(SigFindNearest);
+}
+
+#if WITH_EDITORONLY_DATA
+bool UNiagaraDataInterfaceFrostySensorGrid::GetFunctionHLSL(const FNiagaraDataInterfaceGPUParamInfo& ParamInfo, const FNiagaraDataInterfaceGeneratedFunction& FunctionInfo, int FunctionInstanceIndex, FString& OutHLSL)
+{
+	if ( (FunctionInfo.DefinitionName == NDIFrostySensorGridLocal::UpdateSensorName) ||
+		 (FunctionInfo.DefinitionName == NDIFrostySensorGridLocal::FindNearestName) )
+	{
+		return true;
+	}
+
+	return false;
+}
+
+void UNiagaraDataInterfaceFrostySensorGrid::GetParameterDefinitionHLSL(const FNiagaraDataInterfaceGPUParamInfo& ParamInfo, FString& OutHLSL)
+{
+	TMap<FString, FStringFormatArg> TemplateArgs =
+	{
+		{TEXT("ParameterName"),	ParamInfo.DataInterfaceHLSLSymbol},
+	};
+
+	FString TemplateFile;
+	LoadShaderSourceFile(NDIFrostySensorGridLocal::TemplateShaderFile, EShaderPlatform::SP_PCD3D_SM5, &TemplateFile, nullptr);
+	OutHLSL += FString::Format(*TemplateFile, TemplateArgs);
+}
+
+void UNiagaraDataInterfaceFrostySensorGrid::GetCommonHLSL(FString& OutHlsl)
+{
+	OutHlsl.Appendf(TEXT("#include \"%s\"\n"), NDIFrostySensorGridLocal::CommonShaderFile);
+}
+
+bool UNiagaraDataInterfaceFrostySensorGrid::AppendCompileHash(FNiagaraCompileHashVisitor* InVisitor) const
+{
+	if (!Super::AppendCompileHash(InVisitor))
+	{
+		return false;
+	}
+
+	InVisitor->UpdateString(TEXT("NDIFrostySensorGridCommonHLSLSource"), GetShaderFileHash(NDIFrostySensorGridLocal::CommonShaderFile, EShaderPlatform::SP_PCD3D_SM5).ToString());
+	InVisitor->UpdateString(TEXT("NDIFrostySensorGridTemplateHLSLSource"), GetShaderFileHash(NDIFrostySensorGridLocal::TemplateShaderFile, EShaderPlatform::SP_PCD3D_SM5).ToString());
+
+	return true;
+}
+
+void UNiagaraDataInterfaceFrostySensorGrid::ModifyCompilationEnvironment(EShaderPlatform ShaderPlatform, FShaderCompilerEnvironment& OutEnvironment) const
+{
+	Super::ModifyCompilationEnvironment(ShaderPlatform, OutEnvironment);
+}
+
+
+bool UNiagaraDataInterfaceFrostySensorGrid::UpgradeFunctionCall(FNiagaraFunctionSignature& FunctionSignature)
+{
+	bool WasChanged = false;
+
+	if (FunctionSignature.Name == NDIFrostySensorGridLocal::UpdateSensorName && FunctionSignature.FunctionVersion == NDIFrostySensorGridLocal::NodeVersion::NodeVersion_Initial)
+	{
+		FunctionSignature.AddInput(FNiagaraVariable(FNiagaraTypeDefinition::GetBoolDef(), TEXT("IsValid")), NDIFrostySensorGridLocal::SensorValidDescription);
+		FunctionSignature.FunctionVersion = NDIFrostySensorGridLocal::NodeVersion_Latest;
+
+		WasChanged = true;
+	}
+
+	if (FunctionSignature.Name == NDIFrostySensorGridLocal::FindNearestName && FunctionSignature.FunctionVersion == NDIFrostySensorGridLocal::NodeVersion::NodeVersion_Initial)
+	{
+		FunctionSignature.Outputs.Reset();
+		FunctionSignature.OutputDescriptions.Reset();
+
+		FunctionSignature.AddOutput(FNiagaraVariable(FNiagaraTypeDefinition::GetVec3Def(), TEXT("Out_Location")), NDIFrostySensorGridLocal::OutputLocationDescription);
+		FunctionSignature.AddOutput(FNiagaraVariable(FNiagaraTypeDefinition::GetFloatDef(), TEXT("Out_Distance")), NDIFrostySensorGridLocal::OutputDistanceDescription);
+		FunctionSignature.AddOutput(FNiagaraVariable(FNiagaraTypeDefinition::GetBoolDef(), TEXT("Out_IsValid")), NDIFrostySensorGridLocal::OutputIsValidDescription);
+		FunctionSignature.FunctionVersion = NDIFrostySensorGridLocal::NodeVersion_Latest;
+
+		WasChanged = true;
+	}
+
+	return WasChanged;
+}
+
+#endif
+
+int32 UNiagaraDataInterfaceFrostySensorGrid::PerInstanceDataSize() const
+{
+	return 4;
+}
+
+bool UNiagaraDataInterfaceFrostySensorGrid::InitPerInstanceData(void* PerInstanceData, FNiagaraSystemInstance* SystemInstance)
+{
+	if (SystemInstance)
+	{
+		FNiagaraDataIntefaceProxyFrostySensorGrid* RT_Proxy = GetProxyAs<FNiagaraDataIntefaceProxyFrostySensorGrid>();
+
+		// Push Updates to Proxy, first release any resources
+		ENQUEUE_RENDER_COMMAND(FDestroyInstance)(
+			[RT_Proxy, RT_Batcher = SystemInstance->GetBatcher(), RT_InstanceID = SystemInstance->GetId()](FRHICommandListImmediate& RHICmdList)
+		{
+			RT_Proxy->RegisterNetworkInstance(RT_Batcher, RT_InstanceID);
+		});
+	}
+
+	return true;
+}
+
+void UNiagaraDataInterfaceFrostySensorGrid::DestroyPerInstanceData(void* PerInstanceData, FNiagaraSystemInstance* SystemInstance)
+{
+	if (SystemInstance)
+	{
+		FNiagaraDataIntefaceProxyFrostySensorGrid* RT_Proxy = GetProxyAs<FNiagaraDataIntefaceProxyFrostySensorGrid>();
+
+		// Push Updates to Proxy, first release any resources
+		ENQUEUE_RENDER_COMMAND(FDestroyInstance)(
+			[RT_Proxy, RT_Batcher = SystemInstance->GetBatcher(), RT_InstanceID = SystemInstance->GetId()](FRHICommandListImmediate& RHICmdList)
+		{
+			RT_Proxy->UnregisterNetworkInstance(RT_Batcher, RT_InstanceID);
+		});
+	}
+}
+
+bool UNiagaraDataInterfaceFrostySensorGrid::Equals(const UNiagaraDataInterface* Other) const
+{
+	if (!Super::Equals(Other))
+	{
+		return false;
+	}
+
+	const UNiagaraDataInterfaceFrostySensorGrid* OtherTyped = CastChecked<const UNiagaraDataInterfaceFrostySensorGrid>(Other);
+	return OtherTyped->SensorCountPerSide == SensorCountPerSide
+		&& OtherTyped->GlobalSensorAccuracy == GlobalSensorAccuracy
+		&& OtherTyped->GlobalSensorRange == GlobalSensorRange;
+}
+
+bool UNiagaraDataInterfaceFrostySensorGrid::CopyToInternal(UNiagaraDataInterface* Destination) const
+{
+	if (!Super::CopyToInternal(Destination))
+	{
+		return false;
+	}
+
+	UNiagaraDataInterfaceFrostySensorGrid* OtherTyped = CastChecked<UNiagaraDataInterfaceFrostySensorGrid>(Destination);
+	OtherTyped->SensorCountPerSide = SensorCountPerSide;
+	OtherTyped->GlobalSensorAccuracy = GlobalSensorAccuracy;
+	OtherTyped->GlobalSensorRange = GlobalSensorRange;
+	OtherTyped->MarkRenderDataDirty();
+	return true;
+}
+
+void UNiagaraDataInterfaceFrostySensorGrid::PushToRenderThreadImpl()
+{
+	FNiagaraDataIntefaceProxyFrostySensorGrid* RT_Proxy = GetProxyAs<FNiagaraDataIntefaceProxyFrostySensorGrid>();
+
+	// Push Updates to Proxy, first release any resources
+	ENQUEUE_RENDER_COMMAND(FUpdateDI)(
+		[RT_Proxy, RT_SensorCount = SensorCountPerSide, RT_GlobalSensorRange = FVector2D(GlobalSensorAccuracy, GlobalSensorRange)](FRHICommandListImmediate& RHICmdList)
+		{
+			RT_Proxy->RenderThreadInitialize(RT_SensorCount, RT_GlobalSensorRange);
+		});
+}
+
+#if WITH_EDITOR
+void UNiagaraDataInterfaceFrostySensorGrid::PostEditChangeProperty(struct FPropertyChangedEvent& PropertyChangedEvent)
+{
+	Super::PostEditChangeProperty(PropertyChangedEvent);
+
+	if (PropertyChangedEvent.Property &&
+		(PropertyChangedEvent.Property->GetFName() == GET_MEMBER_NAME_CHECKED(UNiagaraDataInterfaceFrostySensorGrid, SensorCountPerSide)
+		|| PropertyChangedEvent.Property->GetFName() == GET_MEMBER_NAME_CHECKED(UNiagaraDataInterfaceFrostySensorGrid, GlobalSensorAccuracy)
+		|| PropertyChangedEvent.Property->GetFName() == GET_MEMBER_NAME_CHECKED(UNiagaraDataInterfaceFrostySensorGrid, GlobalSensorRange)))
+	{
+		MarkRenderDataDirty();
+	}
+}
+#endif
+
+//////////////////////////////////////////////////////////////////////////
+
+struct FNiagaraDataInterfaceParametersCS_FrostySesnorGrid : public FNiagaraDataInterfaceParametersCS
+{
+	DECLARE_TYPE_LAYOUT(FNiagaraDataInterfaceParametersCS_FrostySesnorGrid, NonVirtual);
+public:
+	void Bind(const FNiagaraDataInterfaceGPUParamInfo& ParameterInfo, const class FShaderParameterMap& ParameterMap)
+	{
+		BoundsHierarchyParam.Bind(ParameterMap, *(NDIFrostySensorGridLocal::BoundsHierarchyParamName + ParameterInfo.DataInterfaceHLSLSymbol));
+		SensorInfoParam.Bind(ParameterMap, *(NDIFrostySensorGridLocal::SensorInfoParamName + ParameterInfo.DataInterfaceHLSLSymbol));
+		SensorGridDimensionsParam.Bind(ParameterMap, *(NDIFrostySensorGridLocal::SensorGridDimensionsParamName + ParameterInfo.DataInterfaceHLSLSymbol));
+		SensorGridWriteIndexParam.Bind(ParameterMap, *(NDIFrostySensorGridLocal::SensorGridWriteIndexParamName + ParameterInfo.DataInterfaceHLSLSymbol));
+		SensorGridReadIndexParam.Bind(ParameterMap, *(NDIFrostySensorGridLocal::SensorGridReadIndexParamName + ParameterInfo.DataInterfaceHLSLSymbol));
+	}
+
+	void Set(FRHICommandList& RHICmdList, const FNiagaraDataInterfaceSetArgs& Context) const
+	{
+		check(IsInRenderingThread());
+
+		FNiagaraDataIntefaceProxyFrostySensorGrid* Proxy = (FNiagaraDataIntefaceProxyFrostySensorGrid*)Context.DataInterface;
+		FRHIComputeShader* ComputeShaderRHI = Context.Shader.GetComputeShader();
+
+		FSensorGridNetworkProxy* NetworkProxy = Proxy->GetNetwork(Context.Batcher);
+
+		NetworkProxy->BeginSimulation(RHICmdList);
+
+		if (BoundsHierarchyParam.IsBound())
+		{
+			RHICmdList.Transition(FRHITransitionInfo(NetworkProxy->BoundsHierarchy->UAV, ERHIAccess::Unknown, ERHIAccess::UAVCompute));
+			RHICmdList.SetUAVParameter(ComputeShaderRHI, BoundsHierarchyParam.GetUAVIndex(), NetworkProxy->BoundsHierarchy->UAV);
+		}
+
+		if (SensorInfoParam.IsBound())
+		{
+			FRHIShaderResourceView* SRV = NetworkProxy->EmptyResults->SRV;
+			if (NetworkProxy->SensorInfo)
+			{
+				SRV = NetworkProxy->SensorInfo->SRV;
+			}
+			SetSRVParameter(RHICmdList, ComputeShaderRHI, SensorInfoParam, SRV);
+		}
+
+		SetShaderValue(RHICmdList, ComputeShaderRHI, SensorGridDimensionsParam, FIntVector(NetworkProxy->SensorGridDimensions.X, NetworkProxy->SensorGridDimensions.Y, NetworkProxy->AllocatedOwnerCount));
+
+		{
+			const int32* WriteIndex = NetworkProxy->InstanceOwnerWriteIndexMap.Find(Context.SystemInstanceID);
+			SetShaderValue(RHICmdList, ComputeShaderRHI, SensorGridWriteIndexParam, WriteIndex ? *WriteIndex : INDEX_NONE);
+		}
+
+		{
+			const int32* ReadIndex = NetworkProxy->InstanceOwnerReadIndexMap.Find(Context.SystemInstanceID);
+			SetShaderValue(RHICmdList, ComputeShaderRHI, SensorGridReadIndexParam, ReadIndex ? *ReadIndex : INDEX_NONE);
+		}
+	}
+
+	void Unset(FRHICommandList& RHICmdList, const FNiagaraDataInterfaceSetArgs& Context) const
+	{
+		FRHIComputeShader* ComputeShaderRHI = Context.Shader.GetComputeShader();
+		FNiagaraDataIntefaceProxyFrostySensorGrid* Proxy = (FNiagaraDataIntefaceProxyFrostySensorGrid*)Context.DataInterface;
+
+		if (BoundsHierarchyParam.IsUAVBound())
+		{
+			BoundsHierarchyParam.UnsetUAV(RHICmdList, ComputeShaderRHI);
+		}
+
+		if (FSensorGridNetworkProxy* NetworkProxy = Proxy->GetNetwork(Context.Batcher))
+		{
+			NetworkProxy->EndSimulation(RHICmdList);
+		}
+	}
+
+private:
+	LAYOUT_FIELD(FRWShaderParameter, BoundsHierarchyParam);
+	LAYOUT_FIELD(FShaderResourceParameter, SensorInfoParam);
+	LAYOUT_FIELD(FShaderParameter, SensorGridDimensionsParam);
+	LAYOUT_FIELD(FShaderParameter, SensorGridWriteIndexParam);
+	LAYOUT_FIELD(FShaderParameter, SensorGridReadIndexParam);
+};
+
+IMPLEMENT_TYPE_LAYOUT(FNiagaraDataInterfaceParametersCS_FrostySesnorGrid);
+
+IMPLEMENT_NIAGARA_DI_PARAMETER(UNiagaraDataInterfaceFrostySensorGrid, FNiagaraDataInterfaceParametersCS_FrostySesnorGrid);
+
+#undef LOCTEXT_NAMESPACE
\ No newline at end of file
diff --git a/Engine/Plugins/Runtime/FrostySensorGrid/Source/FrostySensorGrid/Private/NiagaraDataInterfaceFrostySensorGrid.h b/Engine/Plugins/Runtime/FrostySensorGrid/Source/FrostySensorGrid/Private/NiagaraDataInterfaceFrostySensorGrid.h
new file mode 100644
--- /dev/null
+++ b/Engine/Plugins/Runtime/FrostySensorGrid/Source/FrostySensorGrid/Private/NiagaraDataInterfaceFrostySensorGrid.h
@@ -0,0 +1,59 @@
+// Copyright Epic Games, Inc. All Rights Reserved.
+#pragma once
+
+#include "NiagaraCommon.h"
+#include "NiagaraDataInterface.h"
+#include "NiagaraDataInterfaceFrostySensorGrid.generated.h"
+
+/** Data Interface allowing spatial queries on 3d particles uniquely projected onto a 2d grid */
+UCLASS(EditInlineNew, meta = (DisplayName = "Frosty Sensor Grid"))
+class UNiagaraDataInterfaceFrostySensorGrid : public UNiagaraDataInterface
+{
+	GENERATED_UCLASS_BODY()
+public:
+
+	DECLARE_NIAGARA_DI_PARAMETER();
+
+	UPROPERTY(EditAnywhere, Category = "Sensors")
+	int32 SensorCountPerSide = 0;
+
+	UPROPERTY(EditAnywhere, Category = "Sensors")
+	float GlobalSensorAccuracy = 0.0f;
+
+	UPROPERTY(EditAnywhere, Category = "Sensors")
+	float GlobalSensorRange = 0.0f;
+
+	//UObject Interface
+	virtual void PostInitProperties() override;
+	virtual void PostLoad() override;
+#if WITH_EDITOR
+	virtual void PostEditChangeProperty(struct FPropertyChangedEvent& PropertyChangedEvent) override;
+#endif
+	//UObject Interface End
+
+	virtual void GetFunctions(TArray<FNiagaraFunctionSignature>& OutFunctions) override;
+	virtual void PushToRenderThreadImpl() override;
+
+	virtual bool CanExecuteOnTarget(ENiagaraSimTarget Target) const override { return Target == ENiagaraSimTarget::GPUComputeSim; }
+
+#if WITH_EDITORONLY_DATA
+	virtual bool AppendCompileHash(FNiagaraCompileHashVisitor* InVisitor) const override;
+	virtual void ModifyCompilationEnvironment(EShaderPlatform ShaderPlatform, struct FShaderCompilerEnvironment& OutEnvironment) const override;
+	virtual void GetCommonHLSL(FString& OutHLSL) override;
+	virtual void GetParameterDefinitionHLSL(const FNiagaraDataInterfaceGPUParamInfo& ParamInfo, FString& OutHLSL) override;
+	virtual bool GetFunctionHLSL(const FNiagaraDataInterfaceGPUParamInfo& ParamInfo, const FNiagaraDataInterfaceGeneratedFunction& FunctionInfo, int FunctionInstanceIndex, FString& OutHLSL) override;
+	virtual bool UpgradeFunctionCall(FNiagaraFunctionSignature& FunctionSignature) override;
+#endif
+
+	virtual bool HasPreSimulateTick() const override { return true; }
+	virtual bool HasPostSimulateTick() const override { return true; }
+
+	virtual int32 PerInstanceDataSize() const override;
+	virtual bool InitPerInstanceData(void* PerInstanceData, FNiagaraSystemInstance* SystemInstance) override;
+	virtual void DestroyPerInstanceData(void* PerInstanceData, FNiagaraSystemInstance* SystemInstance) override;
+
+	virtual bool Equals(const UNiagaraDataInterface* Other) const override;
+
+protected:
+	virtual bool CopyToInternal(UNiagaraDataInterface* Destination) const override;
+};
diff --git a/Engine/Plugins/Runtime/FrostySensorGrid/Source/FrostySensorGrid/Public/FrostySensorGrid.h b/Engine/Plugins/Runtime/FrostySensorGrid/Source/FrostySensorGrid/Public/FrostySensorGrid.h
new file mode 100644
--- /dev/null
+++ b/Engine/Plugins/Runtime/FrostySensorGrid/Source/FrostySensorGrid/Public/FrostySensorGrid.h
@@ -0,0 +1,22 @@
+// Copyright Epic Games, Inc. All Rights Reserved.
+
+#pragma once
+
+#include "Modules/ModuleInterface.h"
+#include "Modules/ModuleManager.h"
+
+/** Implements the HairStrands module  */
+class FFrostySensorGrid : public IModuleInterface
+{
+public:
+
+	//~ IModuleInterface interface
+
+	virtual void StartupModule() override;
+	virtual void ShutdownModule() override;
+
+	virtual bool SupportsDynamicReloading() override
+	{
+		return true;
+	}
+};
diff --git a/Engine/Plugins/Runtime/FrostySensorGrid/Source/FrostySensorGridShaders/FrostySensorGridShaders.Build.cs b/Engine/Plugins/Runtime/FrostySensorGrid/Source/FrostySensorGridShaders/FrostySensorGridShaders.Build.cs
new file mode 100644
--- /dev/null
+++ b/Engine/Plugins/Runtime/FrostySensorGrid/Source/FrostySensorGridShaders/FrostySensorGridShaders.Build.cs
@@ -0,0 +1,52 @@
+// Copyright Epic Games, Inc. All Rights Reserved.
+
+using UnrealBuildTool;
+
+public class FrostySensorGridShaders : ModuleRules
+{
+    public FrostySensorGridShaders(ReadOnlyTargetRules Target) : base(Target)
+    {
+        PrivateDependencyModuleNames.AddRange(
+            new string[] {
+                "Core",
+                "CoreUObject",
+                "Engine",
+                "NiagaraCore",
+                "Renderer",
+            }
+        );
+
+        PublicDependencyModuleNames.AddRange(
+            new string[] {
+                "Core",
+                "RenderCore",
+                "VectorVM",
+                "RHI",
+                "Projects",
+                "NiagaraCore"
+            }
+        );
+
+        if (Target.bBuildEditor == true)
+        {
+            PublicDependencyModuleNames.AddRange(
+                new string[] {
+                    "TargetPlatform",
+                });
+
+            PrivateDependencyModuleNames.AddRange(
+                new string[] {
+            });
+        }
+
+        PublicIncludePathModuleNames.AddRange(
+            new string[] {
+            });
+
+        PrivateIncludePathModuleNames.AddRange(
+            new string[] {
+                "DerivedDataCache",
+                "Niagara",
+            });
+    }
+}
diff --git a/Engine/Plugins/Runtime/FrostySensorGrid/Source/FrostySensorGridShaders/Private/FrostySensorGridShaders.cpp b/Engine/Plugins/Runtime/FrostySensorGrid/Source/FrostySensorGridShaders/Private/FrostySensorGridShaders.cpp
new file mode 100644
--- /dev/null
+++ b/Engine/Plugins/Runtime/FrostySensorGrid/Source/FrostySensorGridShaders/Private/FrostySensorGridShaders.cpp
@@ -0,0 +1,128 @@
+// Copyright Epic Games, Inc. All Rights Reserved.
+
+#include "FrostySensorGridShaders.h"
+#include "ShaderPermutation.h"
+#include "ShaderParameterStruct.h"
+
+class FFrostySensorGridAggregateBoundsCs : public FGlobalShader
+{
+	DECLARE_GLOBAL_SHADER(FFrostySensorGridAggregateBoundsCs);
+	SHADER_USE_PARAMETER_STRUCT(FFrostySensorGridAggregateBoundsCs, FGlobalShader);
+
+public:
+	BEGIN_SHADER_PARAMETER_STRUCT(FParameters, FROSTYSENSORGRIDSHADERS_API)
+		SHADER_PARAMETER_UAV(Buffer, BoundsHierarchy)
+		SHADER_PARAMETER(FIntVector, TargetStrideAndOffset)
+		SHADER_PARAMETER(FIntVector, SourceStrideAndOffset)
+	END_SHADER_PARAMETER_STRUCT()
+};
+
+IMPLEMENT_GLOBAL_SHADER(FFrostySensorGridAggregateBoundsCs, "/FrostySensorGrid/FrostySensorGridAggregateBounds.usf", "AggregateBounds", SF_Compute);
+
+class FFrostySensorGridFindClosestSensorsCs : public FGlobalShader
+{
+	DECLARE_GLOBAL_SHADER(FFrostySensorGridFindClosestSensorsCs);
+	SHADER_USE_PARAMETER_STRUCT(FFrostySensorGridFindClosestSensorsCs, FGlobalShader);
+
+public:
+	static const int32 MaxHierarchyLevels = 16;
+	static const int32 NetworkCountPerGroup = 8;
+
+	BEGIN_SHADER_PARAMETER_STRUCT(FParameters, FROSTYSENSORGRIDSHADERS_API)
+		SHADER_PARAMETER_SRV(Buffer<float4>, SensorHierarchy)
+		SHADER_PARAMETER_UAV(StructuredBuffer<FSensorInfo>, NearestSensors)
+		SHADER_PARAMETER_ARRAY(int32, HierarchyLevelOffsets, [MaxHierarchyLevels])
+		SHADER_PARAMETER(FVector2D, DistanceBounds)
+		SHADER_PARAMETER(FIntVector, SensorGridDimensions)
+		SHADER_PARAMETER(int32, HierarchyLevelCount)
+	END_SHADER_PARAMETER_STRUCT()
+
+	static void ModifyCompilationEnvironment(const FGlobalShaderPermutationParameters& Parameters, FShaderCompilerEnvironment& OutEnvironment)
+	{
+		FGlobalShader::ModifyCompilationEnvironment(Parameters, OutEnvironment);
+		OutEnvironment.SetDefine(TEXT("MAX_HIERARCHY_LEVEL_COUNT"), MaxHierarchyLevels);
+		OutEnvironment.SetDefine(TEXT("NETWORK_COUNT_PER_GROUP"), NetworkCountPerGroup);
+	}
+};
+
+IMPLEMENT_GLOBAL_SHADER(FFrostySensorGridFindClosestSensorsCs, "/FrostySensorGrid/FrostySensorGridFindClosest.usf", "FindClosestSensors", SF_Compute);
+
+FFrostySensorGridHelper::FFrostySensorGridHelper(ERHIFeatureLevel::Type InFeatureLevel, const FIntVector& InSensorGridDimensions)
+	: FeatureLevel(InFeatureLevel)
+	, SensorGridDimensions(InSensorGridDimensions)
+{
+
+}
+
+void FFrostySensorGridHelper::BuildBounds(FRHICommandList& RHICmdList, FRHIUnorderedAccessView* HierarchyUAv)
+{
+	TShaderMapRef<FFrostySensorGridAggregateBoundsCs> AggregateBoundsShader(GetGlobalShaderMap(FeatureLevel));
+	FRHIComputeShader* ShaderRHI = AggregateBoundsShader.GetComputeShader();
+
+	const uint32 LayerCount = FMath::FloorLog2(FMath::Max<uint32>(SensorGridDimensions.X, SensorGridDimensions.Y));
+
+	if (LayerCount)
+	{
+		RHICmdList.SetComputeShader(ShaderRHI);
+
+		const int32 BaseSensorGridSize = SensorGridDimensions.X * SensorGridDimensions.Y * SensorGridDimensions.Z;
+
+		FIntPoint SourceGrid(SensorGridDimensions.X, SensorGridDimensions.Y);
+		int32 SourceOffset = 0;
+
+		FIntPoint TargetGrid(SourceGrid.X >> 1, SourceGrid.Y >> 1);
+		int32 TargetOffset = SourceGrid.X * SourceGrid.Y * SensorGridDimensions.Z;
+
+		for (uint32 LayerIt = 0; LayerIt < LayerCount; ++LayerIt)
+		{
+			if (LayerIt)
+			{
+				RHICmdList.Transition(FRHITransitionInfo(HierarchyUAv, ERHIAccess::Unknown, ERHIAccess::UAVCompute));
+			}
+
+			FFrostySensorGridAggregateBoundsCs::FParameters PassParameters;
+			PassParameters.BoundsHierarchy = HierarchyUAv;
+			PassParameters.SourceStrideAndOffset = FIntVector(SourceGrid.X, SourceGrid.Y, SourceOffset * sizeof(FVector4));
+			PassParameters.TargetStrideAndOffset = FIntVector(TargetGrid.X, TargetGrid.Y, TargetOffset * sizeof(FVector4));
+
+			const FIntVector ThreadGroupCount = FIntVector(TargetGrid.X, TargetGrid.Y, SensorGridDimensions.Z);
+
+			SetShaderParameters(RHICmdList, AggregateBoundsShader, ShaderRHI, PassParameters);
+			RHICmdList.DispatchComputeShader(ThreadGroupCount.X, ThreadGroupCount.Y, ThreadGroupCount.Z);
+
+			SourceGrid = TargetGrid;
+			SourceOffset = TargetOffset;
+
+			TargetGrid = FIntPoint(TargetGrid.X >> 1, TargetGrid.Y >> 1);
+			TargetOffset += SourceGrid.X * SourceGrid.Y * SensorGridDimensions.Z;
+		}
+
+		UnsetShaderUAVs(RHICmdList, AggregateBoundsShader, ShaderRHI);
+	}
+}
+
+void FFrostySensorGridHelper::FindNearestSensors(FRHICommandList& RHICmdList, FRHIShaderResourceView* HierarchySrv, const FVector2D& GlobalSensorRange, FRHIUnorderedAccessView* ResultsUav)
+{
+	TShaderMapRef<FFrostySensorGridFindClosestSensorsCs> FindClosestSensorShader(GetGlobalShaderMap(FeatureLevel));
+	FRHIComputeShader* ShaderRHI = FindClosestSensorShader.GetComputeShader();
+
+	FFrostySensorGridFindClosestSensorsCs::FParameters PassParameters;
+	PassParameters.DistanceBounds = GlobalSensorRange;
+	PassParameters.SensorHierarchy = HierarchySrv;
+	PassParameters.NearestSensors = ResultsUav;
+	PassParameters.SensorGridDimensions = SensorGridDimensions;
+	PassParameters.HierarchyLevelCount = 1 + FMath::FloorLog2(FMath::Max<uint32>(SensorGridDimensions.X, SensorGridDimensions.Y));
+
+	int32 Offset = 0;
+
+	for (int32 LayerIt = 0; LayerIt < PassParameters.HierarchyLevelCount; ++LayerIt)
+	{
+		PassParameters.HierarchyLevelOffsets[LayerIt] = Offset * sizeof(FVector4);
+		Offset += SensorGridDimensions.Z * (SensorGridDimensions.X >> LayerIt) * (SensorGridDimensions.Y >> LayerIt);
+	}
+
+	RHICmdList.SetComputeShader(ShaderRHI);
+	SetShaderParameters(RHICmdList, FindClosestSensorShader, ShaderRHI, PassParameters);
+	RHICmdList.DispatchComputeShader(SensorGridDimensions.X >> 1, SensorGridDimensions.Y >> 1, SensorGridDimensions.Z);
+	UnsetShaderUAVs(RHICmdList, FindClosestSensorShader, ShaderRHI);
+}
\ No newline at end of file
diff --git a/Engine/Plugins/Runtime/FrostySensorGrid/Source/FrostySensorGridShaders/Private/FrostySensorGridShadersModule.cpp b/Engine/Plugins/Runtime/FrostySensorGrid/Source/FrostySensorGridShaders/Private/FrostySensorGridShadersModule.cpp
new file mode 100644
--- /dev/null
+++ b/Engine/Plugins/Runtime/FrostySensorGrid/Source/FrostySensorGridShaders/Private/FrostySensorGridShadersModule.cpp
@@ -0,0 +1,16 @@
+// Copyright Epic Games, Inc. All Rights Reserved.
+
+#include "FrostySensorGridShadersModule.h"
+#include "Misc/Paths.h"
+#include "Modules/ModuleManager.h"
+#include "Interfaces/IPluginManager.h"
+#include "ShaderCore.h"
+
+IMPLEMENT_MODULE(IFrostySensorGridShadersModule, FrostySensorGridShaders);
+
+void IFrostySensorGridShadersModule::StartupModule()
+{
+	// Maps virtual shader source directory /Plugin/FX/Niagara to the plugin's actual Shaders directory.
+	FString PluginShaderDir = FPaths::Combine(IPluginManager::Get().FindPlugin(TEXT("FrostySensorGrid"))->GetBaseDir(), TEXT("Shaders"));
+	AddShaderSourceDirectoryMapping(TEXT("/FrostySensorGrid"), PluginShaderDir);
+}
\ No newline at end of file
diff --git a/Engine/Plugins/Runtime/FrostySensorGrid/Source/FrostySensorGridShaders/Public/FrostySensorGridShaders.h b/Engine/Plugins/Runtime/FrostySensorGrid/Source/FrostySensorGridShaders/Public/FrostySensorGridShaders.h
new file mode 100644
--- /dev/null
+++ b/Engine/Plugins/Runtime/FrostySensorGrid/Source/FrostySensorGridShaders/Public/FrostySensorGridShaders.h
@@ -0,0 +1,23 @@
+// Copyright Epic Games, Inc. All Rights Reserved.
+#pragma once
+
+#include "CoreMinimal.h"
+
+#include "Math/IntVector.h"
+#include "RHIDefinitions.h"
+
+class FRHIShaderResourceView;
+class FRHIUnorderedAccessView;
+
+class FROSTYSENSORGRIDSHADERS_API FFrostySensorGridHelper
+{
+public:
+	FFrostySensorGridHelper(ERHIFeatureLevel::Type InFeatureLevel, const FIntVector& InSensorGridDimensions);
+
+	void BuildBounds(FRHICommandList& RHICmdList, FRHIUnorderedAccessView* HierarchyUAv);
+	void FindNearestSensors(FRHICommandList& RHICmdList, FRHIShaderResourceView* HierarchySrv, const FVector2D& GlobalSensorRange, FRHIUnorderedAccessView* ResultsUav);
+
+private:
+	const ERHIFeatureLevel::Type FeatureLevel;
+	const FIntVector SensorGridDimensions;
+};
\ No newline at end of file
diff --git a/Engine/Plugins/Runtime/FrostySensorGrid/Source/FrostySensorGridShaders/Public/FrostySensorGridShadersModule.h b/Engine/Plugins/Runtime/FrostySensorGrid/Source/FrostySensorGridShaders/Public/FrostySensorGridShadersModule.h
new file mode 100644
--- /dev/null
+++ b/Engine/Plugins/Runtime/FrostySensorGrid/Source/FrostySensorGridShaders/Public/FrostySensorGridShadersModule.h
@@ -0,0 +1,18 @@
+// Copyright Epic Games, Inc. All Rights Reserved.
+
+#pragma once
+
+#include "CoreMinimal.h"
+#include "Modules/ModuleInterface.h"
+
+/**
+* FrostySensorGridShader module interface
+*/
+class IFrostySensorGridShadersModule : public IModuleInterface
+{
+public:
+	virtual void StartupModule() override;
+
+private:
+};
+
